Fixed blank AEC template asking to be updated.

#rnx

diff --git a/Templates/TP_AEC_BlankBP/TP_AEC_BlankBP.uproject b/Templates/TP_AEC_BlankBP/TP_AEC_BlankBP.uproject
--- a/Templates/TP_AEC_BlankBP/TP_AEC_BlankBP.uproject
+++ b/Templates/TP_AEC_BlankBP/TP_AEC_BlankBP.uproject
@@ -19,7 +19,10 @@
 		},
 		{
 			"Name": "DatasmithCADImporter",
-			"Enabled": true
+			"Enabled": true,
+			"SupportedTargetPlatforms": [
+				"Win64"
+			]
 		},
 		{
 			"Name": "VariantManager",
