AEC and MFG Blank templates

enabling correct list of plugins

diff --git a/Templates/TP_AEC_BlankBP/TP_AEC_BlankBP.uproject b/Templates/TP_AEC_BlankBP/TP_AEC_BlankBP.uproject
--- a/Templates/TP_AEC_BlankBP/TP_AEC_BlankBP.uproject
+++ b/Templates/TP_AEC_BlankBP/TP_AEC_BlankBP.uproject
@@ -28,6 +28,10 @@
 				{
 			"Name": "HDRIBackdrop",
 			"Enabled": true
+		},
+		{
+			"Name": "SequencerScripting",
+			"Enabled": false
 		}
 	]
 }
\ No newline at end of file
diff --git a/Templates/TP_MFG_BlankBP/TP_MFG_BlankBP.uproject b/Templates/TP_MFG_BlankBP/TP_MFG_BlankBP.uproject
--- a/Templates/TP_MFG_BlankBP/TP_MFG_BlankBP.uproject
+++ b/Templates/TP_MFG_BlankBP/TP_MFG_BlankBP.uproject
@@ -28,6 +28,10 @@
 				{
 			"Name": "HDRIBackdrop",
 			"Enabled": true
+		},
+		{
+			"Name": "SequencerScripting",
+			"Enabled": false
 		}
 	]
 }
\ No newline at end of file
