Fix for Build DDC issue 212112 Disabled GPULightmass before saving levels in Architecture Templates #rb none #preflight

diff --git a/Templates/TP_AEC_BlankBP/TP_AEC_BlankBP.uproject b/Templates/TP_AEC_BlankBP/TP_AEC_BlankBP.uproject
--- a/Templates/TP_AEC_BlankBP/TP_AEC_BlankBP.uproject
+++ b/Templates/TP_AEC_BlankBP/TP_AEC_BlankBP.uproject
@@ -57,10 +57,6 @@
 			"Name": "HDRIBackdrop",
 			"Enabled": true
 		},
-		{
-			"Name": "GPULightmass",
-			"Enabled": true
-		},
 		{
 			"Name": "PythonScriptPlugin",
 			"Enabled": true
